Adding VRSpectator

diff --git a/Templates/TemplateResources/Standard/VRSpectator/Config/Config.ini b/Templates/TemplateResources/Standard/VRSpectator/Config/Config.ini
new file mode 100644
--- /dev/null
+++ b/Templates/TemplateResources/Standard/VRSpectator/Config/Config.ini
@@ -0,0 +1,33 @@
+[/Script/Engine.InputSettings]
+bAltEnterTogglesFullscreen=True
+bF11TogglesFullscreen=True
+bUseMouseForTouch=False
+bEnableMouseSmoothing=True
+bEnableFOVScaling=True
+bCaptureMouseOnLaunch=True
+bAlwaysShowTouchInterface=False
+bShowConsoleOnFourFingerTap=True
+bEnableGestureRecognizer=False
+bUseAutocorrect=False
+DefaultViewportMouseCaptureMode=CapturePermanently_IncludingInitialMouseDown
+DefaultViewportMouseLockMode=LockOnCapture
+FOVScale=0.011110
+DoubleClickTime=0.200000
++ActionMappings=(ActionName="SpectatorResetFOV",bShift=False,bCtrl=False,bAlt=False,bCmd=False,Key=MiddleMouseButton)
++ActionMappings=(ActionName="SpectatorFadeToggle",bShift=False,bCtrl=False,bAlt=False,bCmd=False,Key=F)
++ActionMappings=(ActionName="SpectatorActiveToggle",bShift=False,bCtrl=False,bAlt=False,bCmd=False,Key=Tab)
++ActionMappings=(ActionName="SpectatorResetRotation",bShift=False,bCtrl=False,bAlt=False,bCmd=False,Key=R)
++AxisMappings=(AxisName="SpectatorYaw",Scale=1.000000,Key=MouseX)
++AxisMappings=(AxisName="SpectatorPitch",Scale=1.000000,Key=MouseY)
++AxisMappings=(AxisName="SpectatorFOV",Scale=1.000000,Key=MouseWheelAxis)
++AxisMappings=(AxisName="SpectatorMoveForward",Scale=1.000000,Key=W)
++AxisMappings=(AxisName="SpectatorMoveRight",Scale=1.000000,Key=D)
++AxisMappings=(AxisName="SpectatorMoveForward",Scale=-1.000000,Key=S)
++AxisMappings=(AxisName="SpectatorMoveRight",Scale=-1.000000,Key=A)
++AxisMappings=(AxisName="SpectatorMoveUp",Scale=1.000000,Key=SpaceBar)
++AxisMappings=(AxisName="SpectatorMoveUp",Scale=1.000000,Key=E)
++AxisMappings=(AxisName="SpectatorMoveUp",Scale=-1.000000,Key=Q)
+DefaultPlayerInputClass=/Script/Engine.PlayerInput
+DefaultInputComponentClass=/Script/Engine.InputComponent
+DefaultTouchInterface=None
++ConsoleKeys=Tilde
\ No newline at end of file
diff --git a/Templates/TemplateResources/Standard/VRSpectator/FeaturePack/manifest.json b/Templates/TemplateResources/Standard/VRSpectator/FeaturePack/manifest.json
new file mode 100644
--- /dev/null
+++ b/Templates/TemplateResources/Standard/VRSpectator/FeaturePack/manifest.json
@@ -0,0 +1,41 @@
+{
+	"Version": 1,
+	"Name": [
+		{
+			"Language": "en",
+			"Text": "VR Spectator"
+		}
+	],
+	"Description": [
+		{
+			"Language": "en",
+			"Text": "Feature Pack for VR Spectator"
+		}
+	],
+	"AssetTypes": [
+		{
+			"Language": "en",
+			"Text": ""
+		}
+	],
+	"SearchTags": [
+		{
+			"Language": "en",
+			"Text": ""
+		}
+	],
+	"ClassTypes": " ",
+	"Category": "SharedPack",
+	"Thumbnail": "VRSpectator.png",
+	"Screenshots": [
+		"VRSpectator_Preview.png"
+	],
+	"Ident": "VRSpectatorStandard",
+	"AdditionalFiles": {
+		"DestinationFilesFolder": "VRSpectator",
+		"AdditionalFilesList": [
+			"Templates/TemplateResources/Standard/VRSpectator/Content/*.*"
+		]
+	}
+}
+
diff --git a/Templates/TemplateResources/Standard/VRSpectator/contents.txt b/Templates/TemplateResources/Standard/VRSpectator/contents.txt
new file mode 100644
--- /dev/null
+++ b/Templates/TemplateResources/Standard/VRSpectator/contents.txt
@@ -0,0 +1,4 @@
+../../../Templates/TemplateResources/Standard/VRSpectator/FeaturePack/manifest.json
+../../../Templates/TemplateResources/Standard/VRSpectator/Config/config.ini
+../../../Templates/TemplateResources/Standard/VRSpectator/Media/VRSpectator.png
+../../../Templates/TemplateResources/Standard/VRSpectator/Media/VRSpectator_Preview.png
