Added ArchvisDefault TemplateResources pack #preflight none

diff --git a/Templates/TemplateResources/Standard/ArchvisDefault/FeaturePack/manifest.json b/Templates/TemplateResources/Standard/ArchvisDefault/FeaturePack/manifest.json
new file mode 100644
--- /dev/null
+++ b/Templates/TemplateResources/Standard/ArchvisDefault/FeaturePack/manifest.json
@@ -0,0 +1,48 @@
+{
+	"Version": 1,
+	"Name":
+	[
+		{
+			"Language": "en",
+			"Text": "Archvis templates Shared Resource Standard"
+		}
+	],
+	"Description":
+	[
+		{
+			"Language": "en",
+			"Text": "Shared Resource Pack for Archvis templates"
+		}
+	],
+	"AssetTypes":
+	[
+		{
+			"Language": "en",
+			"Text": ""
+		}
+	],
+	"SearchTags":
+	[
+		{
+			"Language": "en",
+			"Text": ""
+		}
+	],	
+	"ClassTypes": " ",
+	"Category": "SharedPack",
+	"Thumbnail": "ArchvisDefault.png",
+	"Screenshots":
+	[
+		"ArchvisDefault_Preview.png"
+	],
+	"Ident": "ArchvisDefaultStandard",
+	"AdditionalFiles":
+	{
+		"DestinationFilesFolder": "ArchvisDefault",
+		"AdditionalFilesList":
+		[
+			"Templates/TemplateResources/Standard/ArchvisDefault/Content/*.*"
+		]
+	}
+}
+
