Added VR Spectator as an additional Feature Pack to include when the VR Template is added as a feature pack to a project

diff --git a/Templates/TemplateResources/Standard/VRSpectator/Config/Config.ini b/Templates/TemplateResources/Standard/VRSpectator/Config/Config.ini
--- a/Templates/TemplateResources/Standard/VRSpectator/Config/Config.ini
+++ b/Templates/TemplateResources/Standard/VRSpectator/Config/Config.ini
@@ -1,18 +1,5 @@
-[/Script/Engine.InputSettings]
-bAltEnterTogglesFullscreen=True
-bF11TogglesFullscreen=True
-bUseMouseForTouch=False
-bEnableMouseSmoothing=True
-bEnableFOVScaling=True
-bCaptureMouseOnLaunch=True
-bAlwaysShowTouchInterface=False
-bShowConsoleOnFourFingerTap=True
-bEnableGestureRecognizer=False
-bUseAutocorrect=False
-DefaultViewportMouseCaptureMode=CapturePermanently_IncludingInitialMouseDown
-DefaultViewportMouseLockMode=LockOnCapture
-FOVScale=0.011110
-DoubleClickTime=0.200000
+[InputSettings]
+
 +ActionMappings=(ActionName="SpectatorResetFOV",bShift=False,bCtrl=False,bAlt=False,bCmd=False,Key=MiddleMouseButton)
 +ActionMappings=(ActionName="SpectatorFadeToggle",bShift=False,bCtrl=False,bAlt=False,bCmd=False,Key=F)
 +ActionMappings=(ActionName="SpectatorActiveToggle",bShift=False,bCtrl=False,bAlt=False,bCmd=False,Key=Tab)
@@ -26,8 +13,4 @@ DoubleClickTime=0.200000
 +AxisMappings=(AxisName="SpectatorMoveRight",Scale=-1.000000,Key=A)
 +AxisMappings=(AxisName="SpectatorMoveUp",Scale=1.000000,Key=SpaceBar)
 +AxisMappings=(AxisName="SpectatorMoveUp",Scale=1.000000,Key=E)
-+AxisMappings=(AxisName="SpectatorMoveUp",Scale=-1.000000,Key=Q)
-DefaultPlayerInputClass=/Script/Engine.PlayerInput
-DefaultInputComponentClass=/Script/Engine.InputComponent
-DefaultTouchInterface=None
-+ConsoleKeys=Tilde
\ No newline at end of file
++AxisMappings=(AxisName="SpectatorMoveUp",Scale=-1.000000,Key=Q)
\ No newline at end of file
diff --git a/Templates/TemplateResources/Standard/VRSpectator/FeaturePack/manifest.json b/Templates/TemplateResources/Standard/VRSpectator/FeaturePack/manifest.json
--- a/Templates/TemplateResources/Standard/VRSpectator/FeaturePack/manifest.json
+++ b/Templates/TemplateResources/Standard/VRSpectator/FeaturePack/manifest.json
@@ -25,7 +25,7 @@
 		}
 	],
 	"ClassTypes": " ",
-	"Category": "SharedPack",
+	"Category": ["SharedPack", "BlueprintFeature"],
 	"Thumbnail": "VRSpectator.png",
 	"Screenshots": [
 		"VRSpectator_Preview.png"
