Removing VR Spectator as a Feature pack for 5.0 - will revisit sharedpack/featurepack compatibility for 5.1

diff --git a/Templates/TemplateResources/Standard/VRSpectator/FeaturePack/manifest.json b/Templates/TemplateResources/Standard/VRSpectator/FeaturePack/manifest.json
--- a/Templates/TemplateResources/Standard/VRSpectator/FeaturePack/manifest.json
+++ b/Templates/TemplateResources/Standard/VRSpectator/FeaturePack/manifest.json
@@ -25,7 +25,7 @@
 		}
 	],
 	"ClassTypes": " ",
-	"Category": ["SharedPack", "BlueprintFeature"],
+	"Category": "SharedPack",
 	"Thumbnail": "VRSpectator.png",
 	"Screenshots": [
 		"VRSpectator_Preview.png"
