[Templates] removed an old map from Top down bp template. Added manifest and pngs for shared cursor feature pack.

diff --git a/Templates/TemplateResources/Standard/Cursor/FeaturePack/manifest.json b/Templates/TemplateResources/Standard/Cursor/FeaturePack/manifest.json
new file mode 100644
--- /dev/null
+++ b/Templates/TemplateResources/Standard/Cursor/FeaturePack/manifest.json
@@ -0,0 +1,48 @@
+{
+	"Version": 1,
+	"Name":
+	[
+		{
+			"Language": "en",
+			"Text": "Cursor shared resource"
+		}
+	],
+	"Description":
+	[
+		{
+			"Language": "en",
+			"Text": "Shared resource pack for Cursor"
+		}
+	],
+	"AssetTypes":
+	[
+		{
+			"Language": "en",
+			"Text": "StaticMesh, Material, FX"
+		}
+	],
+	"SearchTags":
+	[
+		{
+			"Language": "en",
+			"Text": "Cursor"
+		}
+	],	
+	"ClassTypes": "StaticMesh",
+	"Category": "SharedPack",
+	"Thumbnail": "Cursor.png",
+	"Screenshots":
+	[
+		"Cursor_preview.png"
+	],
+	"Ident": "CursorStandard",
+	"AdditionalFiles":
+	{
+		"DestinationFilesFolder": "Cursor",
+		"AdditionalFilesList":
+		[
+			"Templates/TemplateResources/Standard/Cursor/Content/*.*"
+		]
+	}
+}
+
