Adding missing manifest and media files for shared packs for Vehicle template.

diff --git a/Templates/TemplateResources/High/Track/FeaturePack/manifest.json b/Templates/TemplateResources/High/Track/FeaturePack/manifest.json
new file mode 100644
--- /dev/null
+++ b/Templates/TemplateResources/High/Track/FeaturePack/manifest.json
@@ -0,0 +1,48 @@
+{
+	"Version": 1,
+	"Name":
+	[
+		{
+			"Language": "en",
+			"Text": "Track Shared Resource"
+		}
+	],
+	"Description":
+	[
+		{
+			"Language": "en",
+			"Text": "Shared Resource Pack for Vehicle Template"
+		}
+	],
+	"AssetTypes":
+	[
+		{
+			"Language": "en",
+			"Text": ""
+		}
+	],
+	"SearchTags":
+	[
+		{
+			"Language": "en",
+			"Text": ""
+		}
+	],	
+	"ClassTypes": " ",
+	"Category": "SharedPack",
+	"Thumbnail": "Track.png",
+	"Screenshots":
+	[
+		"Track_Preview.png"
+	],
+	"Ident": "TrackStandard",
+	"AdditionalFiles":
+	{
+		"DestinationFilesFolder": "Track",
+		"AdditionalFilesList":
+		[
+			"Templates/TemplateResources/Standard/Track/Content/*.*"
+		]
+	}
+}
+
diff --git a/Templates/TemplateResources/Standard/Vehicles/FeaturePack/manifest.json b/Templates/TemplateResources/Standard/Vehicles/FeaturePack/manifest.json
new file mode 100644
--- /dev/null
+++ b/Templates/TemplateResources/Standard/Vehicles/FeaturePack/manifest.json
@@ -0,0 +1,48 @@
+{
+	"Version": 1,
+	"Name":
+	[
+		{
+			"Language": "en",
+			"Text": "Vehicles Shared Resource"
+		}
+	],
+	"Description":
+	[
+		{
+			"Language": "en",
+			"Text": "Shared Resource Pack for Vehicle Template"
+		}
+	],
+	"AssetTypes":
+	[
+		{
+			"Language": "en",
+			"Text": ""
+		}
+	],
+	"SearchTags":
+	[
+		{
+			"Language": "en",
+			"Text": ""
+		}
+	],	
+	"ClassTypes": " ",
+	"Category": "SharedPack",
+	"Thumbnail": "Vehicles.png",
+	"Screenshots":
+	[
+		"Vehicles_Preview.png"
+	],
+	"Ident": "VehiclesStandard",
+	"AdditionalFiles":
+	{
+		"DestinationFilesFolder": "Vehicles",
+		"AdditionalFilesList":
+		[
+			"Templates/TemplateResources/Standard/Vehicles/Content/*.*"
+		]
+	}
+}
+
