Templates - TopDown
Added missing manifest and images to Template resource packs.
#rnx

diff --git a/Templates/TemplateResources/Standard/Variant_Strategy/FeaturePack/manifest.json b/Templates/TemplateResources/Standard/Variant_Strategy/FeaturePack/manifest.json
new file mode 100644
--- /dev/null
+++ b/Templates/TemplateResources/Standard/Variant_Strategy/FeaturePack/manifest.json
@@ -0,0 +1,48 @@
+{
+	"Version": 1,
+	"Name":
+	[
+		{
+			"Language": "en",
+			"Text": "Strategy Variant Shared Resource Standard"
+		}
+	],
+	"Description":
+	[
+		{
+			"Language": "en",
+			"Text": "Shared Resource Pack for TopDown Strategy Variant"
+		}
+	],
+	"AssetTypes":
+	[
+		{
+			"Language": "en",
+			"Text": ""
+		}
+	],
+	"SearchTags":
+	[
+		{
+			"Language": "en",
+			"Text": ""
+		}
+	],	
+	"ClassTypes": " ",
+	"Category": "SharedPack",
+	"Thumbnail": "Variant_Strategy.png",
+	"Screenshots":
+	[
+		"Variant_Strategy_Preview.png"
+	],
+	"Ident": "Variant_StrategyStandard",
+	"AdditionalFiles":
+	{
+		"DestinationFilesFolder": "Variant_Strategy",
+		"AdditionalFilesList":
+		[
+			"Templates/TemplateResources/Standard/Variant_Strategy/Content/*.*"
+		]
+	}
+}
+
diff --git a/Templates/TemplateResources/Standard/Variant_TwinStick/FeaturePack/manifest.json b/Templates/TemplateResources/Standard/Variant_TwinStick/FeaturePack/manifest.json
new file mode 100644
--- /dev/null
+++ b/Templates/TemplateResources/Standard/Variant_TwinStick/FeaturePack/manifest.json
@@ -0,0 +1,48 @@
+{
+	"Version": 1,
+	"Name":
+	[
+		{
+			"Language": "en",
+			"Text": "TwinStick Variant Shared Resource Standard"
+		}
+	],
+	"Description":
+	[
+		{
+			"Language": "en",
+			"Text": "Shared Resource Pack for TopDown TwinStick Variant"
+		}
+	],
+	"AssetTypes":
+	[
+		{
+			"Language": "en",
+			"Text": ""
+		}
+	],
+	"SearchTags":
+	[
+		{
+			"Language": "en",
+			"Text": ""
+		}
+	],	
+	"ClassTypes": " ",
+	"Category": "SharedPack",
+	"Thumbnail": "Variant_TwinStick.png",
+	"Screenshots":
+	[
+		"Variant_TwinStick_Preview.png"
+	],
+	"Ident": "Variant_TwinStickStandard",
+	"AdditionalFiles":
+	{
+		"DestinationFilesFolder": "Variant_TwinStick",
+		"AdditionalFilesList":
+		[
+			"Templates/TemplateResources/Standard/Variant_TwinStick/Content/*.*"
+		]
+	}
+}
+
