Converted levels back to OFPA.
Levels can no longer be opened in main template project.
Fixed animation issues from Third Person Character

diff --git a/Templates/TemplateResources/Standard/Variant_Combat/FeaturePack/manifest.json b/Templates/TemplateResources/Standard/Variant_Combat/FeaturePack/manifest.json
--- a/Templates/TemplateResources/Standard/Variant_Combat/FeaturePack/manifest.json
+++ b/Templates/TemplateResources/Standard/Variant_Combat/FeaturePack/manifest.json
@@ -39,9 +39,10 @@
 	"AdditionalFiles":
 	{
 		"DestinationFilesFolder": "Variant_Combat",
-		"AdditionalFilesList":
-		[
-			"Templates/TemplateResources/Standard/Variant_Combat/Content/*.*"
+		"AdditionalFilesList": [
+			"Templates/TemplateResources/Standard/Variant_Combat/Content/*.*",
+			"Templates/TemplateResources/Standard/Variant_Combat/__ExternalActors__/*.*",
+			"Templates/TemplateResources/Standard/Variant_Combat/__ExternalObjects__/*.*"
 		]
 	}
 }
diff --git a/Templates/TemplateResources/Standard/Variant_Offroad/FeaturePack/manifest.json b/Templates/TemplateResources/Standard/Variant_Offroad/FeaturePack/manifest.json
--- a/Templates/TemplateResources/Standard/Variant_Offroad/FeaturePack/manifest.json
+++ b/Templates/TemplateResources/Standard/Variant_Offroad/FeaturePack/manifest.json
@@ -34,7 +34,9 @@
 	"AdditionalFiles": {
 		"DestinationFilesFolder": "Variant_Offroad",
 		"AdditionalFilesList": [
-			"Templates/TemplateResources/Standard/Variant_Offroad/Content/*.*"
+			"Templates/TemplateResources/Standard/Variant_Offroad/Content/*.*",
+			"Templates/TemplateResources/Standard/Variant_Offroad/__ExternalActors__/*.*",
+			"Templates/TemplateResources/Standard/Variant_Offroad/__ExternalObjects__/*.*"
 		]
 	}
 }
diff --git a/Templates/TemplateResources/Standard/Variant_Platforming/FeaturePack/manifest.json b/Templates/TemplateResources/Standard/Variant_Platforming/FeaturePack/manifest.json
--- a/Templates/TemplateResources/Standard/Variant_Platforming/FeaturePack/manifest.json
+++ b/Templates/TemplateResources/Standard/Variant_Platforming/FeaturePack/manifest.json
@@ -39,9 +39,10 @@
 	"AdditionalFiles":
 	{
 		"DestinationFilesFolder": "Variant_Platforming",
-		"AdditionalFilesList":
-		[
-			"Templates/TemplateResources/Standard/Variant_Platforming/Content/*.*"
+		"AdditionalFilesList": [
+			"Templates/TemplateResources/Standard/Variant_Platforming/Content/*.*",
+			"Templates/TemplateResources/Standard/Variant_Platforming/__ExternalActors__/*.*",
+			"Templates/TemplateResources/Standard/Variant_Platforming/__ExternalObjects__/*.*"
 		]
 	}
 }
diff --git a/Templates/TemplateResources/Standard/Variant_Shooter/FeaturePack/manifest.json b/Templates/TemplateResources/Standard/Variant_Shooter/FeaturePack/manifest.json
--- a/Templates/TemplateResources/Standard/Variant_Shooter/FeaturePack/manifest.json
+++ b/Templates/TemplateResources/Standard/Variant_Shooter/FeaturePack/manifest.json
@@ -34,7 +34,9 @@
 	"AdditionalFiles": {
 		"DestinationFilesFolder": "Variant_Shooter",
 		"AdditionalFilesList": [
-			"Templates/TemplateResources/Standard/Variant_Shooter/Content/*.*"
+			"Templates/TemplateResources/Standard/Variant_Shooter/Content/*.*",
+			"Templates/TemplateResources/Standard/Variant_Shooter/__ExternalActors__/*.*",
+			"Templates/TemplateResources/Standard/Variant_Shooter/__ExternalObjects__/*.*"
 		]
 	}
 }
diff --git a/Templates/TemplateResources/Standard/Variant_SideScroller/FeaturePack/manifest.json b/Templates/TemplateResources/Standard/Variant_SideScroller/FeaturePack/manifest.json
--- a/Templates/TemplateResources/Standard/Variant_SideScroller/FeaturePack/manifest.json
+++ b/Templates/TemplateResources/Standard/Variant_SideScroller/FeaturePack/manifest.json
@@ -39,9 +39,10 @@
 	"AdditionalFiles":
 	{
 		"DestinationFilesFolder": "Variant_SideScroller",
-		"AdditionalFilesList":
-		[
-			"Templates/TemplateResources/Standard/Variant_SideScroller/Content/*.*"
+		"AdditionalFilesList": [
+			"Templates/TemplateResources/Standard/Variant_SideScroller/Content/*.*",
+			"Templates/TemplateResources/Standard/Variant_Sidescroller/__ExternalActors__/*.*",
+			"Templates/TemplateResources/Standard/Variant_Sidescroller/__ExternalObjects__/*.*"
 		]
 	}
 }
diff --git a/Templates/TemplateResources/Standard/Variant_Strategy/FeaturePack/manifest.json b/Templates/TemplateResources/Standard/Variant_Strategy/FeaturePack/manifest.json
--- a/Templates/TemplateResources/Standard/Variant_Strategy/FeaturePack/manifest.json
+++ b/Templates/TemplateResources/Standard/Variant_Strategy/FeaturePack/manifest.json
@@ -39,9 +39,10 @@
 	"AdditionalFiles":
 	{
 		"DestinationFilesFolder": "Variant_Strategy",
-		"AdditionalFilesList":
-		[
-			"Templates/TemplateResources/Standard/Variant_Strategy/Content/*.*"
+		"AdditionalFilesList": [
+			"Templates/TemplateResources/Standard/Variant_Strategy/Content/*.*",
+			"Templates/TemplateResources/Standard/Variant_Strategy/__ExternalActors__/*.*",
+			"Templates/TemplateResources/Standard/Variant_Strategy/__ExternalObjects__/*.*"
 		]
 	}
 }
diff --git a/Templates/TemplateResources/Standard/Variant_Timetrial/FeaturePack/manifest.json b/Templates/TemplateResources/Standard/Variant_Timetrial/FeaturePack/manifest.json
--- a/Templates/TemplateResources/Standard/Variant_Timetrial/FeaturePack/manifest.json
+++ b/Templates/TemplateResources/Standard/Variant_Timetrial/FeaturePack/manifest.json
@@ -34,7 +34,9 @@
 	"AdditionalFiles": {
 		"DestinationFilesFolder": "Variant_Timetrial",
 		"AdditionalFilesList": [
-			"Templates/TemplateResources/Standard/Variant_Timetrial/Content/*.*"
+			"Templates/TemplateResources/Standard/Variant_Timetrial/Content/*.*",
+			"Templates/TemplateResources/Standard/Variant_Timetrial/__ExternalActors__/*.*",
+			"Templates/TemplateResources/Standard/Variant_Timetrial/__ExternalObjects__/*.*"
 		]
 	}
 }
diff --git a/Templates/TemplateResources/Standard/Variant_TwinStick/FeaturePack/manifest.json b/Templates/TemplateResources/Standard/Variant_TwinStick/FeaturePack/manifest.json
--- a/Templates/TemplateResources/Standard/Variant_TwinStick/FeaturePack/manifest.json
+++ b/Templates/TemplateResources/Standard/Variant_TwinStick/FeaturePack/manifest.json
@@ -39,9 +39,10 @@
 	"AdditionalFiles":
 	{
 		"DestinationFilesFolder": "Variant_TwinStick",
-		"AdditionalFilesList":
-		[
-			"Templates/TemplateResources/Standard/Variant_TwinStick/Content/*.*"
+		"AdditionalFilesList": [
+			"Templates/TemplateResources/Standard/Variant_TwinStick/Content/*.*",
+			"Templates/TemplateResources/Standard/Variant_TwinStick/__ExternalActors__/*.*",
+			"Templates/TemplateResources/Standard/Variant_TwinStick/__ExternalObjects__/*.*"
 		]
 	}
 }
