VR Template | Submitting XR Mannequins template resources content that was not added for submit

#virtualized

diff --git a/Templates/TemplateResources/Standard/XRMannequins/FeaturePack/manifest.json b/Templates/TemplateResources/Standard/XRMannequins/FeaturePack/manifest.json
new file mode 100644
--- /dev/null
+++ b/Templates/TemplateResources/Standard/XRMannequins/FeaturePack/manifest.json
@@ -0,0 +1,41 @@
+{
+	"Version": 1,
+	"Name": [
+		{
+			"Language": "en",
+			"Text": "XR Mannequins"
+		}
+	],
+	"Description": [
+		{
+			"Language": "en",
+			"Text": "Feature Pack for XR Mannequin Hands"
+		}
+	],
+	"AssetTypes": [
+		{
+			"Language": "en",
+			"Text": ""
+		}
+	],
+	"SearchTags": [
+		{
+			"Language": "en",
+			"Text": ""
+		}
+	],
+	"ClassTypes": " ",
+	"Category": "SharedPack",
+	"Thumbnail": "VRSpectator.png",
+	"Screenshots": [
+		"VRSpectator_Preview.png"
+	],
+	"Ident": "XRMannequinsStandard",
+	"AdditionalFiles": {
+		"DestinationFilesFolder": "XRMannequins",
+		"AdditionalFilesList": [
+			"Templates/TemplateResources/Standard/XRMannequins/Content/*.*"
+		]
+	}
+}
+
diff --git a/Templates/TemplateResources/Standard/XRMannequins/contents.txt b/Templates/TemplateResources/Standard/XRMannequins/contents.txt
new file mode 100644
--- /dev/null
+++ b/Templates/TemplateResources/Standard/XRMannequins/contents.txt
@@ -0,0 +1,4 @@
+../../../Templates/TemplateResources/Standard/XRMannequins/FeaturePack/manifest.json
+../../../Templates/TemplateResources/Standard/XRMannequins/Config/config.ini
+../../../Templates/TemplateResources/Standard/XRMannequins/Media/VRSpectator.png
+../../../Templates/TemplateResources/Standard/XRMannequins/Media/VRSpectator_Preview.png
