VR Template | Updated XRMannequins template resources media names. Removed VRSpectator and XRMannequins Config.ini as they were only used to add legacy input actions.

diff --git a/Templates/TemplateResources/Standard/VRSpectator/Config/Config.ini b/Templates/TemplateResources/Standard/VRSpectator/Config/Config.ini
deleted file mode 100644
diff --git a/Templates/TemplateResources/Standard/VRSpectator/contents.txt b/Templates/TemplateResources/Standard/VRSpectator/contents.txt
--- a/Templates/TemplateResources/Standard/VRSpectator/contents.txt
+++ b/Templates/TemplateResources/Standard/VRSpectator/contents.txt
@@ -1,4 +1,3 @@
 ../../../Templates/TemplateResources/Standard/VRSpectator/FeaturePack/manifest.json
-../../../Templates/TemplateResources/Standard/VRSpectator/Config/config.ini
 ../../../Templates/TemplateResources/Standard/VRSpectator/Media/VRSpectator.png
-../../../Templates/TemplateResources/Standard/VRSpectator/Media/VRSpectator_Preview.png
+../../../Templates/TemplateResources/Standard/VRSpectator/Media/VRSpectator_Preview.png
\ No newline at end of file
diff --git a/Templates/TemplateResources/Standard/XRMannequins/Config/Config.ini b/Templates/TemplateResources/Standard/XRMannequins/Config/Config.ini
deleted file mode 100644
diff --git a/Templates/TemplateResources/Standard/XRMannequins/FeaturePack/manifest.json b/Templates/TemplateResources/Standard/XRMannequins/FeaturePack/manifest.json
--- a/Templates/TemplateResources/Standard/XRMannequins/FeaturePack/manifest.json
+++ b/Templates/TemplateResources/Standard/XRMannequins/FeaturePack/manifest.json
@@ -26,9 +26,9 @@
 	],
 	"ClassTypes": " ",
 	"Category": "SharedPack",
-	"Thumbnail": "VRSpectator.png",
+	"Thumbnail": "XRMannequins.png",
 	"Screenshots": [
-		"VRSpectator_Preview.png"
+		"XRMannequins_Preview.png"
 	],
 	"Ident": "XRMannequinsStandard",
 	"AdditionalFiles": {
diff --git a/Templates/TemplateResources/Standard/XRMannequins/contents.txt b/Templates/TemplateResources/Standard/XRMannequins/contents.txt
--- a/Templates/TemplateResources/Standard/XRMannequins/contents.txt
+++ b/Templates/TemplateResources/Standard/XRMannequins/contents.txt
@@ -1,4 +1,3 @@
 ../../../Templates/TemplateResources/Standard/XRMannequins/FeaturePack/manifest.json
-../../../Templates/TemplateResources/Standard/XRMannequins/Config/config.ini
-../../../Templates/TemplateResources/Standard/XRMannequins/Media/VRSpectator.png
-../../../Templates/TemplateResources/Standard/XRMannequins/Media/VRSpectator_Preview.png
+../../../Templates/TemplateResources/Standard/XRMannequins/Media/XRMannequins.png
+../../../Templates/TemplateResources/Standard/XRMannequins/Media/XRMannequins_Preview.png
\ No newline at end of file
