TP_ThirdPerson, SideScroller variant
- Fixed variant configuration error that caused a blank map on Linux
#rnx

diff --git a/Templates/TemplateResources/Standard/Variant_SideScroller/FeaturePack/manifest.json b/Templates/TemplateResources/Standard/Variant_SideScroller/FeaturePack/manifest.json
--- a/Templates/TemplateResources/Standard/Variant_SideScroller/FeaturePack/manifest.json
+++ b/Templates/TemplateResources/Standard/Variant_SideScroller/FeaturePack/manifest.json
@@ -41,8 +41,8 @@
 		"DestinationFilesFolder": "Variant_SideScroller",
 		"AdditionalFilesList": [
 			"Templates/TemplateResources/Standard/Variant_SideScroller/Content/*.*",
-			"Templates/TemplateResources/Standard/Variant_Sidescroller/__ExternalActors__/*.*",
-			"Templates/TemplateResources/Standard/Variant_Sidescroller/__ExternalObjects__/*.*"
+			"Templates/TemplateResources/Standard/Variant_SideScroller/__ExternalActors__/*.*",
+			"Templates/TemplateResources/Standard/Variant_SideScroller/__ExternalObjects__/*.*"
 		]
 	}
 }
